Editor: Hardware target settings application WIP
- Add a settings builder/gatherer class that either builds a description of what the meta-settings will change, or actually applies the changes
- Display the pending changes
- First pass on actual settings controlled by the hardware+graphics level meta-settings (very WIP)

diff --git a/Engine/Source/Editor/DetailCustomizations/Private/HardwareTargetingSettingsDetails.cpp b/Engine/Source/Editor/DetailCustomizations/Private/HardwareTargetingSettingsDetails.cpp
--- a/Engine/Source/Editor/DetailCustomizations/Private/HardwareTargetingSettingsDetails.cpp
+++ b/Engine/Source/Editor/DetailCustomizations/Private/HardwareTargetingSettingsDetails.cpp
@@ -27,29 +27,46 @@ void FHardwareTargetingSettingsDetails::CustomizeDetails(IDetailLayoutBuilder& D
 			return Settings->HasPendingChanges() ? EVisibility::Visible : EVisibility::Collapsed;
 		})
 		[
-			SNew(SHorizontalBox)
-
-			+ SHorizontalBox::Slot()
-			.AutoWidth()
-			.VAlign(VAlign_Center)
-			[
-				SNew(STextBlock)
-				.Text(LOCTEXT("RestartMessage", "These changes will be applied when this project is reopened."))
-			]
-
-			+ SHorizontalBox::Slot()
-			.AutoWidth()
-			.Padding(6)
+			SNew(SVerticalBox)
+			+ SVerticalBox::Slot()
+			.AutoHeight()
 			[
-				SNew(SButton)
-				.Text(LOCTEXT("ApplyNow", "Apply Now"))
-				.OnClicked_Static([](){
+				SNew(SHorizontalBox)
+
+				+ SHorizontalBox::Slot()
+				.AutoWidth()
+				.VAlign(VAlign_Center)
+				[
+					SNew(STextBlock)
+					.Text(LOCTEXT("RestartMessage", "The following changes will be applied when this project is reopened."))
+				]
+
+				+ SHorizontalBox::Slot()
+				.AutoWidth()
+				.Padding(6)
+				[
+					SNew(SButton)
+					.Text(LOCTEXT("ApplyNow", "Apply Now"))
+					.OnClicked_Static([](){
 				
-					IHardwareTargetingModule& Module = IHardwareTargetingModule::Get();
-					Module.ApplyHardwareTargetingSettings();
+						IHardwareTargetingModule& Module = IHardwareTargetingModule::Get();
+						Module.ApplyHardwareTargetingSettings();
 
-					return FReply::Handled();
-				})
+						return FReply::Handled();
+					})
+				]
+			]
+			+ SVerticalBox::Slot()
+			.FillHeight(1.0f)
+			[
+				SNew(SRichTextBlock)
+				.AutoWrapText(true)
+				.Justification(ETextJustify::Left)
+				.Text(TAttribute<FText>::Create(TAttribute<FText>::FGetter::CreateStatic([](){
+					return IHardwareTargetingModule::Get().QueryReadableDescriptionOfHardwareTargetingSettings();
+				})))
+				.TextStyle(FEditorStyle::Get(), "HardwareTargets.Normal")
+ 				.DecoratorStyleSet(&FEditorStyle::Get())
 			]
 		]
 	];
diff --git a/Engine/Source/Editor/EditorStyle/Private/SlateEditorStyle.cpp b/Engine/Source/Editor/EditorStyle/Private/SlateEditorStyle.cpp
--- a/Engine/Source/Editor/EditorStyle/Private/SlateEditorStyle.cpp
+++ b/Engine/Source/Editor/EditorStyle/Private/SlateEditorStyle.cpp
@@ -1401,7 +1401,29 @@ void FSlateEditorStyle::FStyle::SetupGeneralStyles()
 
 		Set("Credits.Hyperlink", HoverOnlyHyperlink);
 	}
-#endif // WITH_EDITOR || IS_PROGRAM
+#endif // WITH_EDITOR
+
+	// Hardware target settings
+#if WITH_EDITOR
+	{
+		FLinearColor EditorOrange = FLinearColor(0.728f, 0.364f, 0.003f);
+
+		FTextBlockStyle TargetSettingsNormal = FTextBlockStyle(NormalText)
+			.SetFont(TTF_CORE_FONT("Fonts/Roboto-Regular", 12))
+			.SetShadowOffset(FVector2D::UnitVector);
+
+		Set("HardwareTargets.Normal", TargetSettingsNormal);
+
+ 		Set("HardwareTargets.Strong", FTextBlockStyle(TargetSettingsNormal)
+ 			.SetFont(TTF_CORE_FONT("Fonts/Roboto-Bold", 12))
+ 			.SetShadowOffset(FVector2D::UnitVector));
+
+		Set("HardwareTargets.H1", FTextBlockStyle(TargetSettingsNormal)
+			.SetColorAndOpacity(EditorOrange)
+			.SetFont(TTF_CORE_FONT("Fonts/Roboto-Bold", 16))
+			.SetShadowOffset(FVector2D::UnitVector));
+	}
+#endif
 
 	// New Level Dialog
 #if WITH_EDITOR || IS_PROGRAM
diff --git a/Engine/Source/Editor/HardwareTargeting/HardwareTargeting.Build.cs b/Engine/Source/Editor/HardwareTargeting/HardwareTargeting.Build.cs
--- a/Engine/Source/Editor/HardwareTargeting/HardwareTargeting.Build.cs
+++ b/Engine/Source/Editor/HardwareTargeting/HardwareTargeting.Build.cs
@@ -22,6 +22,7 @@ public class HardwareTargeting : ModuleRules
 				"EditorStyle",
 				"UnrealEd",
 				"Settings",
+				"EngineSettings",
 			}
 		);
 	}
diff --git a/Engine/Source/Editor/HardwareTargeting/Private/HardwareTargetingModule.cpp b/Engine/Source/Editor/HardwareTargeting/Private/HardwareTargetingModule.cpp
--- a/Engine/Source/Editor/HardwareTargeting/Private/HardwareTargetingModule.cpp
+++ b/Engine/Source/Editor/HardwareTargeting/Private/HardwareTargetingModule.cpp
@@ -7,23 +7,132 @@
 #include "Internationalization.h"
 #include "SDecoratedEnumCombo.h"
 
+#include "Runtime/Engine/Classes/Engine/RendererSettings.h"
+
 #define LOCTEXT_NAMESPACE "HardwareTargeting"
 
+//////////////////////////////////////////////////////////////////////////
+// FMetaSettingGatherer
+
+struct FMetaSettingGatherer
+{
+	TSet<UObject*> DirtyObjects;
+	FTextBuilder DescriptionBuffer;
+
+	// Are we just displaying what would change, or actually changing things?
+	bool bReadOnly;
+
+	FMetaSettingGatherer()
+		: bReadOnly(false)
+	{
+	}
+
+	void AddEntry(UObject* SettingsObject, UProperty* Property, FText NewValue, bool bModified)
+	{
+		if (bModified && !bReadOnly)
+		{
+			DirtyObjects.Add(SettingsObject);
+
+			FPropertyChangedEvent ChangeEvent(Property, false, EPropertyChangeType::ValueSet);
+			SettingsObject->PostEditChangeProperty(ChangeEvent);
+		}
+
+		if (bReadOnly)
+		{
+			FText SettingDisplayName = Property->GetDisplayNameText();
+
+			FFormatNamedArguments Args;
+			Args.Add(TEXT("SettingName"), SettingDisplayName);
+			Args.Add(TEXT("SettingValue"), NewValue);
+
+			FText FormatString = bModified ? LOCTEXT("MetaSettingDisplayStringUnmodified", "{SettingName} is {SettingValue} <HardwareTargets.Strong>(modified)</>") : LOCTEXT("MetaSettingDisplayStringUnmodified", "{SettingName} is {SettingValue}");
+
+			DescriptionBuffer.AppendLine(FText::Format(FormatString, Args));
+		}
+	}
+
+	template <typename ValueType>
+	static FText ValueToString(ValueType Value);
+
+	void Finalize()
+	{
+		check(!bReadOnly);
+		for (UObject* SettingsObject : DirtyObjects)
+		{
+			SettingsObject->UpdateDefaultConfigFile();
+		}
+	}
+
+	void StartSection(FText SectionHeading)
+	{
+		if (bReadOnly)
+		{
+			DescriptionBuffer.AppendLine(SectionHeading);
+			DescriptionBuffer.Indent();
+		}
+	}
+
+	void EndSection()
+	{
+		if (bReadOnly)
+		{
+			DescriptionBuffer.Unindent();
+		}
+	}
+};
+
+
+template <>
+FText FMetaSettingGatherer::ValueToString(bool Value)
+{
+	return Value ? LOCTEXT("BoolEnabled", "enabled") : LOCTEXT("BoolDisabled", "disabled");
+}
+
+template <>
+FText FMetaSettingGatherer::ValueToString(EAntiAliasingMethodUI::Type Value)
+{
+	switch (Value)
+	{
+	case EAntiAliasingMethodUI::AAM_None:
+		return LOCTEXT("AA_None", "None");
+	case EAntiAliasingMethodUI::AAM_FXAA:
+		return LOCTEXT("AA_FXAA", "FXAA");
+	case EAntiAliasingMethodUI::AAM_TemporalAA:
+		return LOCTEXT("AA_TemporalAA", "Temporal AA");
+	default:
+		return FText::AsNumber((int32)Value);
+	}
+}
+
+
+#define UE_META_SETTING_ENTRY(Builder, Class, PropertyName, TargetValue) \
+{ \
+	Class* SettingsObject = GetMutableDefault<Class>(); \
+	bool bModified = SettingsObject->PropertyName != (TargetValue); \
+	if (!Builder.bReadOnly) { SettingsObject->PropertyName = (TargetValue); } \
+	Builder.AddEntry(SettingsObject, FindFieldChecked<UProperty>(Class::StaticClass(), GET_MEMBER_NAME_CHECKED(Class, PropertyName)), FMetaSettingGatherer::ValueToString(TargetValue), bModified); \
+}
+
+//////////////////////////////////////////////////////////////////////////
+// FHardwareTargetingModule
+
 class FHardwareTargetingModule : public IHardwareTargetingModule
 {
+public:
 	// IModuleInterface interface
 	virtual void StartupModule() override;
 	virtual void ShutdownModule() override;
 	// End of IModuleInterface interface
 	
-	/** Apply the current hardware targeting settings if they have changed */
-	virtual void ApplyHardwareTargetingSettings();
-
-	/** Make a new combo box for choosing a hardware class target */
+	// IHardwareTargetingModule interface
+	virtual void ApplyHardwareTargetingSettings() override;
+	virtual FText QueryReadableDescriptionOfHardwareTargetingSettings() override;
 	virtual TSharedRef<SWidget> MakeHardwareClassTargetCombo(FOnHardwareClassChanged OnChanged, TAttribute<EHardwareClass::Type> SelectedEnum) override;
-
-	/** Make a new combo box for choosing a graphics preference */
 	virtual TSharedRef<SWidget> MakeGraphicsPresetTargetCombo(FOnGraphicsPresetChanged OnChanged, TAttribute<EGraphicsPreset::Type> SelectedEnum) override;
+	// End of IHardwareTargetingModule interface
+
+private:
+	void GatherSettings(FMetaSettingGatherer& Builder);
 };
 
 void FHardwareTargetingModule::StartupModule()
@@ -42,31 +151,93 @@ void FHardwareTargetingModule::StartupModule()
 
 void FHardwareTargetingModule::ShutdownModule()
 {
-		
 }
+
+FText FHardwareTargetingModule::QueryReadableDescriptionOfHardwareTargetingSettings()
+{
+	// Gather and stringify the modified settings
+	FMetaSettingGatherer Gatherer;
+	Gatherer.bReadOnly = true;
+	GatherSettings(Gatherer);
+
+	return Gatherer.DescriptionBuffer.ToText();
+}
+
 	
-void FHardwareTargetingModule::ApplyHardwareTargetingSettings()
+void FHardwareTargetingModule::GatherSettings(FMetaSettingGatherer& Builder)
 {
 	UHardwareTargetingSettings* Settings = GetMutableDefault<UHardwareTargetingSettings>();
 
-	int32 EnumValue = 0;
-	bool bFoundKey = GConfig->GetInt(TEXT("AppliedHardwareTargetingSettings"),TEXT("AppliedTargetedHardwareClass"), EnumValue, GEditorIni);
-	if (!bFoundKey || EHardwareClass::Type(EnumValue) != Settings->TargetedHardwareClass)
+	const bool bLowEndMobile = (Settings->TargetedHardwareClass == EHardwareClass::Mobile) && (Settings->DefaultGraphicsPerformance == EGraphicsPreset::Scalable);
+	const bool bAnyMobile = (Settings->TargetedHardwareClass == EHardwareClass::Mobile);
+	const bool bHighEndMobile = (Settings->TargetedHardwareClass == EHardwareClass::Mobile) && (Settings->DefaultGraphicsPerformance == EGraphicsPreset::Maximum);
+	const bool bAnyPC = (Settings->TargetedHardwareClass == EHardwareClass::Desktop);
+	const bool bHighEndPC = (Settings->TargetedHardwareClass == EHardwareClass::Desktop) && (Settings->DefaultGraphicsPerformance == EGraphicsPreset::Maximum);
+
+	{
+		Builder.StartSection(LOCTEXT("RenderingHeading", "<HardwareTargets.H1>Default Rendering Settings</>"));
+
+		// Based roughly on https://docs.unrealengine.com/latest/INT/Platforms/Mobile/PostProcessEffects/index.html
+		UE_META_SETTING_ENTRY(Builder, URendererSettings, bMobileHDR, !bLowEndMobile);
+
+		// Bloom works and isn't terribly expensive on anything beyond low-end
+		UE_META_SETTING_ENTRY(Builder, URendererSettings, bDefaultFeatureBloom, !bLowEndMobile);
+
+		// Separate translucency does nothing in the ES2 renderer
+		UE_META_SETTING_ENTRY(Builder, URendererSettings, bSeparateTranslucency, !bAnyMobile);
+
+		// Motion blur, lens flare, auto-exposure, and ambient occlusion don't work in the ES2 renderer
+		UE_META_SETTING_ENTRY(Builder, URendererSettings, bDefaultFeatureMotionBlur, bHighEndPC);
+		UE_META_SETTING_ENTRY(Builder, URendererSettings, bDefaultFeatureLensFlare, bAnyPC);
+		UE_META_SETTING_ENTRY(Builder, URendererSettings, bDefaultFeatureAutoExposure, bHighEndPC);
+		UE_META_SETTING_ENTRY(Builder, URendererSettings, bDefaultFeatureAmbientOcclusion, bAnyPC);
+
+		// DOF and AA work on mobile but are expensive, keeping them off by default
+		//@TODO: DOF setting doesn't exist yet
+		// UE_META_SETTING_ENTRY(Builder, URendererSettings, bDefaultFeatureDepthOfField, bHighEndPC);
+		UE_META_SETTING_ENTRY(Builder, URendererSettings, DefaultFeatureAntiAliasing, bHighEndPC ? EAntiAliasingMethodUI::AAM_TemporalAA : EAntiAliasingMethodUI::AAM_None);
+
+		Builder.EndSection();
+	}
+
 	{
-		// @todo: Set up project to work with Settings->TargetedHardwareClass
+		Builder.StartSection(LOCTEXT("InputHeading", "<HardwareTargets.H1>Default Input Settings</>"));
+
+		// Mobile uses touch
+		UE_META_SETTING_ENTRY(Builder, UInputSettings, bUseMouseForTouch, bAnyMobile);
+		//@TODO: Use bAlwaysShowTouchInterface (sorta implied by bUseMouseForTouch)?
 
-		GConfig->SetInt(TEXT("AppliedHardwareTargetingSettings"),TEXT("AppliedTargetedHardwareClass"), Settings->TargetedHardwareClass, GEditorIni);
+		Builder.EndSection();
 	}
 
-	bFoundKey = GConfig->GetInt(TEXT("AppliedHardwareTargetingSettings"),TEXT("AppliedDefaultGraphicsPerformance"), EnumValue, GEditorIni);
-	if (!bFoundKey || EGraphicsPreset::Type(EnumValue) != Settings->DefaultGraphicsPerformance)
 	{
-		// @todo: Set up project to work with Settings->DefaultGraphicsPerformance
+		Builder.StartSection(LOCTEXT("GameHeading", "<HardwareTargets.H1>Default Game Settings</>"));
+		
+		// Tablets or phones are usually shared-screen multiplayer instead of split-screen
+		UE_META_SETTING_ENTRY(Builder, UGameMapsSettings, bUseSplitscreen, bAnyPC);
 
-		GConfig->SetInt(TEXT("AppliedHardwareTargetingSettings"),TEXT("AppliedDefaultGraphicsPerformance"), Settings->DefaultGraphicsPerformance, GEditorIni);
+		Builder.EndSection();
 	}
+}
+
+void FHardwareTargetingModule::ApplyHardwareTargetingSettings()
+{
+	UHardwareTargetingSettings* Settings = GetMutableDefault<UHardwareTargetingSettings>();
+
+	// Apply the settings if they've changed
+	if (Settings->HasPendingChanges())
+	{
+		// Gather and apply the modified settings
+		FMetaSettingGatherer Builder;
+		Builder.bReadOnly = false;
+		GatherSettings(Builder);
+		Builder.Finalize();
 
-	GConfig->Flush(false, GEditorIni);
+		// Write out the 'did we apply' values
+		GConfig->SetInt(TEXT("AppliedHardwareTargetingSettings"), TEXT("AppliedDefaultGraphicsPerformance"), Settings->DefaultGraphicsPerformance, GEditorIni);
+		GConfig->SetInt(TEXT("AppliedHardwareTargetingSettings"), TEXT("AppliedTargetedHardwareClass"), Settings->TargetedHardwareClass, GEditorIni);
+		GConfig->Flush(false, GEditorIni);
+	}
 }
 
 TSharedRef<SWidget> FHardwareTargetingModule::MakeHardwareClassTargetCombo(FOnHardwareClassChanged OnChanged, TAttribute<EHardwareClass::Type> SelectedEnum)
@@ -95,46 +266,13 @@ TSharedRef<SWidget> FHardwareTargetingModule::MakeGraphicsPresetTargetCombo(FOnG
 		.OnEnumChanged(OnChanged);
 }
 
-IMPLEMENT_MODULE( FHardwareTargetingModule, HardwareTarget );
-	
 IHardwareTargetingModule& IHardwareTargetingModule::Get()
 {
 	static FHardwareTargetingModule Instance;
 	return Instance;
 }
 
-UHardwareTargetingSettings::UHardwareTargetingSettings( const FPostConstructInitializeProperties& PCIP )
-	: Super(PCIP)
-	, TargetedHardwareClass(EHardwareClass::Desktop)
-	, DefaultGraphicsPerformance(EGraphicsPreset::Maximum)
-{ }
-	
-bool UHardwareTargetingSettings::HasPendingChanges() const
-{
-	UHardwareTargetingSettings* Settings = GetMutableDefault<UHardwareTargetingSettings>();
-
-	int32 EnumValue = 0;
-	if (!GConfig->GetInt(TEXT("AppliedHardwareTargetingSettings"),TEXT("AppliedTargetedHardwareClass"), EnumValue, GEditorIni))
-	{
-		return true;
-	}
-
-	if (EHardwareClass::Type(EnumValue) != Settings->TargetedHardwareClass)
-	{
-		return true;
-	}
-
-	if (!GConfig->GetInt(TEXT("AppliedHardwareTargetingSettings"),TEXT("AppliedDefaultGraphicsPerformance"), EnumValue, GEditorIni))
-	{
-		return true;
-	}
-
-	if (EGraphicsPreset::Type(EnumValue) != Settings->DefaultGraphicsPerformance)
-	{
-		return true;
-	}
-
-	return false;
-}
+IMPLEMENT_MODULE(FHardwareTargetingModule, HardwareTarget);
 
+#undef UE_META_SETTING_ENTRY
 #undef LOCTEXT_NAMESPACE
\ No newline at end of file
diff --git a/Engine/Source/Editor/HardwareTargeting/Private/HardwareTargetingSettings.cpp b/Engine/Source/Editor/HardwareTargeting/Private/HardwareTargetingSettings.cpp
new file mode 100644
--- /dev/null
+++ b/Engine/Source/Editor/HardwareTargeting/Private/HardwareTargetingSettings.cpp
@@ -0,0 +1,40 @@
+// Copyright 1998-2014 Epic Games, Inc. All Rights Reserved.
+
+#include "HardwareTargetingPrivatePCH.h"
+#include "HardwareTargetingModule.h"
+#include "HardwareTargetingSettings.h"
+#include "Settings.h"
+#include "Internationalization.h"
+#include "SDecoratedEnumCombo.h"
+
+UHardwareTargetingSettings::UHardwareTargetingSettings(const FPostConstructInitializeProperties& PCIP)
+	: Super(PCIP)
+	, TargetedHardwareClass(EHardwareClass::Desktop)
+	, DefaultGraphicsPerformance(EGraphicsPreset::Maximum)
+{
+}
+	
+bool UHardwareTargetingSettings::HasPendingChanges() const
+{
+	UHardwareTargetingSettings* Settings = GetMutableDefault<UHardwareTargetingSettings>();
+
+	{
+		int32 AppliedHardwareEnumValue = 0;
+		bool bFoundKey = GConfig->GetInt(TEXT("AppliedHardwareTargetingSettings"), TEXT("AppliedTargetedHardwareClass"), AppliedHardwareEnumValue, GEditorIni);
+		if (!bFoundKey || (EHardwareClass::Type(AppliedHardwareEnumValue) != Settings->TargetedHardwareClass))
+		{
+			return true;
+		}
+	}
+
+	{
+		int32 AppliedGraphicsEnumValue = 0;
+		bool bFoundKey = GConfig->GetInt(TEXT("AppliedHardwareTargetingSettings"), TEXT("AppliedDefaultGraphicsPerformance"), AppliedGraphicsEnumValue, GEditorIni);
+		if (!bFoundKey || (EGraphicsPreset::Type(AppliedGraphicsEnumValue) != Settings->DefaultGraphicsPerformance))
+		{
+			return true;
+		}
+	}
+
+	return false;
+}
diff --git a/Engine/Source/Editor/HardwareTargeting/Public/HardwareTargetingModule.h b/Engine/Source/Editor/HardwareTargeting/Public/HardwareTargetingModule.h
--- a/Engine/Source/Editor/HardwareTargeting/Public/HardwareTargetingModule.h
+++ b/Engine/Source/Editor/HardwareTargeting/Public/HardwareTargetingModule.h
@@ -18,6 +18,9 @@ public:
 	/** Apply the current hardware targeting settings if they have changed */
 	virtual void ApplyHardwareTargetingSettings() = 0;
 
+	/** Queries what the the current hardware targeting settings would like to change */
+	virtual FText QueryReadableDescriptionOfHardwareTargetingSettings() = 0;
+
 	/** Make a new combo box for choosing a hardware class target */
 	virtual TSharedRef<SWidget> MakeHardwareClassTargetCombo(FOnHardwareClassChanged OnChanged, TAttribute<EHardwareClass::Type> SelectedEnum) = 0;
 
