Editor: Add tooltips to 'hardware target' combo boxes and rename settings page "Target Hardware" (was "Hardware Targeting")

diff --git a/Engine/Source/Editor/DetailCustomizations/Private/HardwareTargetingSettingsDetails.cpp b/Engine/Source/Editor/DetailCustomizations/Private/HardwareTargetingSettingsDetails.cpp
--- a/Engine/Source/Editor/DetailCustomizations/Private/HardwareTargetingSettingsDetails.cpp
+++ b/Engine/Source/Editor/DetailCustomizations/Private/HardwareTargetingSettingsDetails.cpp
@@ -196,7 +196,7 @@ public:
 
 void FHardwareTargetingSettingsDetails::CustomizeDetails(IDetailLayoutBuilder& DetailBuilder)
 {	
-	IDetailCategoryBuilder& HardwareTargetingCategory = DetailBuilder.EditCategory(TEXT("Hardware Targeting"));
+	IDetailCategoryBuilder& HardwareTargetingCategory = DetailBuilder.EditCategory(TEXT("Target Hardware"));
 	IDetailCategoryBuilder& PendingChangesCategory = DetailBuilder.EditCategory(TEXT("Pending Changes"));
 
 	auto AnyPendingChangesVisible = []()->EVisibility{
diff --git a/Engine/Source/Editor/HardwareTargeting/Private/HardwareTargetingModule.cpp b/Engine/Source/Editor/HardwareTargeting/Private/HardwareTargetingModule.cpp
--- a/Engine/Source/Editor/HardwareTargeting/Private/HardwareTargetingModule.cpp
+++ b/Engine/Source/Editor/HardwareTargeting/Private/HardwareTargetingModule.cpp
@@ -8,6 +8,7 @@
 #include "SDecoratedEnumCombo.h"
 
 #include "Runtime/Engine/Classes/Engine/RendererSettings.h"
+#include "Editor/Documentation/Public/IDocumentation.h"
 
 #define LOCTEXT_NAMESPACE "HardwareTargeting"
 
@@ -133,7 +134,7 @@ void FHardwareTargetingModule::StartupModule()
 	ISettingsModule& SettingsModule = *ISettingsModule::Get();
 		
 	SettingsModule.RegisterSettings("Project", "Project", "HardwareTargeting",
-		LOCTEXT("HardwareTargetingSettingsName", "Hardware Targeting"),
+		LOCTEXT("HardwareTargetingSettingsName", "Target Hardware"),
 		LOCTEXT("HardwareTargetingSettingsDescription", "Options for choosing which class of hardware to target"),
 		GetMutableDefault<UHardwareTargetingSettings>()
 	);
@@ -251,7 +252,8 @@ TSharedRef<SWidget> FHardwareTargetingModule::MakeHardwareClassTargetCombo(FOnHa
 
 	return SNew(SDecoratedEnumCombo<EHardwareClass::Type>, MoveTemp(HardwareClassInfo))
 		.SelectedEnum(SelectedEnum)
-		.OnEnumChanged(OnChanged);
+		.OnEnumChanged(OnChanged)
+		.ToolTip(IDocumentation::Get()->CreateToolTip(LOCTEXT("HardwareClassTooltip", "Choose the overall class of hardware to target (desktop/console or mobile/tablet)."), NULL, TEXT("Shared/Editor/Settings/TargetHardware"), TEXT("HardwareClass")));
 }
 
 TSharedRef<SWidget> FHardwareTargetingModule::MakeGraphicsPresetTargetCombo(FOnGraphicsPresetChanged OnChanged, TAttribute<EGraphicsPreset::Type> SelectedEnum)
@@ -264,7 +266,8 @@ TSharedRef<SWidget> FHardwareTargetingModule::MakeGraphicsPresetTargetCombo(FOnG
 
 	return SNew(SDecoratedEnumCombo<EGraphicsPreset::Type>, MoveTemp(GraphicsPresetInfo))
 		.SelectedEnum(SelectedEnum)
-		.OnEnumChanged(OnChanged);
+		.OnEnumChanged(OnChanged)
+		.ToolTip(IDocumentation::Get()->CreateToolTip(LOCTEXT("GraphicsPresetTooltip", "Choose the graphical level to target (high-end only or scalable from low-end on up)."), NULL, TEXT("Shared/Editor/Settings/TargetHardware"), TEXT("GraphicalLevel")));
 }
 
 IHardwareTargetingModule& IHardwareTargetingModule::Get()
