Added "Unspecified" options to EGraphicsPreset and EHardwareClass to indicate that no choice has yet been made

These are now the defaults, which prevent any updating of project settings until changed by the user.

diff --git a/Engine/Source/Editor/EditorStyle/Private/SlateEditorStyle.cpp b/Engine/Source/Editor/EditorStyle/Private/SlateEditorStyle.cpp
--- a/Engine/Source/Editor/EditorStyle/Private/SlateEditorStyle.cpp
+++ b/Engine/Source/Editor/EditorStyle/Private/SlateEditorStyle.cpp
@@ -1712,9 +1712,11 @@ void FSlateEditorStyle::FStyle::SetupGeneralStyles()
 	{
 		Set( "HardwareTargeting.MobilePlatform", new IMAGE_BRUSH( "/Icons/HardwareTargeting/Mobile", FVector2D(64, 64) ) );
 		Set( "HardwareTargeting.DesktopPlatform", new IMAGE_BRUSH( "/Icons/HardwareTargeting/Desktop", FVector2D(64, 64) ) );
+		Set( "HardwareTargeting.HardwareUnspecified", new IMAGE_BRUSH( "/Icons/HardwareTargeting/HardwareUnspecified", FVector2D(64, 64) ) );
 
 		Set( "HardwareTargeting.MaximumQuality", new IMAGE_BRUSH( "/Icons/HardwareTargeting/MaximumQuality", FVector2D(64, 64) ) );
 		Set( "HardwareTargeting.ScalableQuality", new IMAGE_BRUSH( "/Icons/HardwareTargeting/ScalableQuality", FVector2D(64, 64) ) );
+		Set( "HardwareTargeting.GraphicsUnspecified", new IMAGE_BRUSH( "/Icons/HardwareTargeting/GraphicsUnspecified", FVector2D(64, 64) ) );
 	}
 
 #endif // WITH_EDITOR
diff --git a/Engine/Source/Editor/HardwareTargeting/Private/HardwareTargetingModule.cpp b/Engine/Source/Editor/HardwareTargeting/Private/HardwareTargetingModule.cpp
--- a/Engine/Source/Editor/HardwareTargeting/Private/HardwareTargetingModule.cpp
+++ b/Engine/Source/Editor/HardwareTargeting/Private/HardwareTargetingModule.cpp
@@ -264,6 +264,8 @@ void FHardwareTargetingModule::ApplyHardwareTargetingSettings()
 TSharedRef<SWidget> FHardwareTargetingModule::MakeHardwareClassTargetCombo(FOnHardwareClassChanged OnChanged, TAttribute<EHardwareClass::Type> SelectedEnum)
 {
 	TArray<SDecoratedEnumCombo<EHardwareClass::Type>::FComboOption> HardwareClassInfo;
+	HardwareClassInfo.Add(SDecoratedEnumCombo<EHardwareClass::Type>::FComboOption(
+		EHardwareClass::Unspecified, FSlateIcon(FEditorStyle::GetStyleSetName(), "HardwareTargeting.HardwareUnspecified"), LOCTEXT("UnspecifiedCaption", "Unspecified"), false));
 	HardwareClassInfo.Add(SDecoratedEnumCombo<EHardwareClass::Type>::FComboOption(
 		EHardwareClass::Desktop, FSlateIcon(FEditorStyle::GetStyleSetName(), "HardwareTargeting.DesktopPlatform"), LOCTEXT("DesktopCaption", "Desktop / Console")));
 	HardwareClassInfo.Add(SDecoratedEnumCombo<EHardwareClass::Type>::FComboOption(
@@ -278,6 +280,8 @@ TSharedRef<SWidget> FHardwareTargetingModule::MakeHardwareClassTargetCombo(FOnHa
 TSharedRef<SWidget> FHardwareTargetingModule::MakeGraphicsPresetTargetCombo(FOnGraphicsPresetChanged OnChanged, TAttribute<EGraphicsPreset::Type> SelectedEnum)
 {
 	TArray<SDecoratedEnumCombo<EGraphicsPreset::Type>::FComboOption> GraphicsPresetInfo;
+	GraphicsPresetInfo.Add(SDecoratedEnumCombo<EGraphicsPreset::Type>::FComboOption(
+		EGraphicsPreset::Unspecified, FSlateIcon(FEditorStyle::GetStyleSetName(), "HardwareTargeting.GraphicsUnspecified"), LOCTEXT("UnspecifiedCaption", "Unspecified"), false));
 	GraphicsPresetInfo.Add(SDecoratedEnumCombo<EGraphicsPreset::Type>::FComboOption(
 		EGraphicsPreset::Maximum, FSlateIcon(FEditorStyle::GetStyleSetName(), "HardwareTargeting.MaximumQuality"), LOCTEXT("MaximumCaption", "Maximum Quality")));
 	GraphicsPresetInfo.Add(SDecoratedEnumCombo<EGraphicsPreset::Type>::FComboOption(
diff --git a/Engine/Source/Editor/HardwareTargeting/Private/HardwareTargetingSettings.cpp b/Engine/Source/Editor/HardwareTargeting/Private/HardwareTargetingSettings.cpp
--- a/Engine/Source/Editor/HardwareTargeting/Private/HardwareTargetingSettings.cpp
+++ b/Engine/Source/Editor/HardwareTargeting/Private/HardwareTargetingSettings.cpp
@@ -9,8 +9,8 @@
 
 UHardwareTargetingSettings::UHardwareTargetingSettings(const FPostConstructInitializeProperties& PCIP)
 	: Super(PCIP)
-	, TargetedHardwareClass(EHardwareClass::Desktop)
-	, DefaultGraphicsPerformance(EGraphicsPreset::Maximum)
+	, TargetedHardwareClass(EHardwareClass::Unspecified)
+	, DefaultGraphicsPerformance(EGraphicsPreset::Unspecified)
 {
 }
 
@@ -18,6 +18,7 @@ bool UHardwareTargetingSettings::HasPendingChanges() const
 {
 	UHardwareTargetingSettings* Settings = GetMutableDefault<UHardwareTargetingSettings>();
 
+	if (Settings->TargetedHardwareClass != EHardwareClass::Unspecified)
 	{
 		FString AppliedHardwareEnumValue;
 		const bool bFoundKey = GConfig->GetString(TEXT("AppliedHardwareTargetingSettings"), TEXT("AppliedTargetedHardwareClass"), AppliedHardwareEnumValue, GEditorIni);
@@ -45,6 +46,7 @@ bool UHardwareTargetingSettings::HasPendingChanges() const
 		}
 	}
 
+	if (Settings->DefaultGraphicsPerformance != EHardwareClass::Unspecified)
 	{
 		FString AppliedGraphicsEnumValue;
 		const bool bFoundKey = GConfig->GetString(TEXT("AppliedHardwareTargetingSettings"), TEXT("AppliedDefaultGraphicsPerformance"), AppliedGraphicsEnumValue, GEditorIni);
diff --git a/Engine/Source/Editor/HardwareTargeting/Private/SDecoratedEnumCombo.h b/Engine/Source/Editor/HardwareTargeting/Private/SDecoratedEnumCombo.h
--- a/Engine/Source/Editor/HardwareTargeting/Private/SDecoratedEnumCombo.h
+++ b/Engine/Source/Editor/HardwareTargeting/Private/SDecoratedEnumCombo.h
@@ -8,13 +8,14 @@ class SDecoratedEnumCombo : public SCompoundWidget
 public:
 	struct FComboOption
 	{
-		FComboOption(TEnumType InValue, FSlateIcon InIcon, FText InText)
-			: Value(InValue), Icon(InIcon), Text(InText)
+		FComboOption(TEnumType InValue, FSlateIcon InIcon, FText InText, bool bInChoosable = true)
+			: Value(InValue), Icon(InIcon), Text(InText), bChoosable(bInChoosable)
 		{}
 
 		TEnumType Value;
 		FSlateIcon Icon;
 		FText Text;
+		bool bChoosable;
 	};
 
 	DECLARE_DELEGATE_OneParam(FOnEnumChanged, TEnumType)
@@ -119,6 +120,10 @@ private:
 		for (int32 Index = 0; Index < Options.Num(); ++Index)
 		{
 			const auto& Option = Options[Index];
+			if (!Option.bChoosable)
+			{
+				continue;
+			}
 
 			HorizontalBox->AddSlot()
 			.AutoWidth()
diff --git a/Engine/Source/Editor/HardwareTargeting/Public/HardwareTargetingSettings.h b/Engine/Source/Editor/HardwareTargeting/Public/HardwareTargetingSettings.h
--- a/Engine/Source/Editor/HardwareTargeting/Public/HardwareTargetingSettings.h
+++ b/Engine/Source/Editor/HardwareTargeting/Public/HardwareTargetingSettings.h
@@ -10,6 +10,9 @@ namespace EHardwareClass
 {
 	enum Type
 	{
+		// Unspecified, meaning no choice has been made yet
+		Unspecified,
+
 		// Desktop or console
 		Desktop,
 		
@@ -24,6 +27,9 @@ namespace EGraphicsPreset
 {
 	enum Type
 	{
+		// Unspecified, meaning no choice has been made yet
+		Unspecified,
+
 		// Maximum quality - High-end features default to enabled
 		Maximum,
 		
