Optimistic build fix for Mac.

diff --git a/Engine/Source/Editor/HardwareTargeting/Private/HardwareTargetingSettings.cpp b/Engine/Source/Editor/HardwareTargeting/Private/HardwareTargetingSettings.cpp
--- a/Engine/Source/Editor/HardwareTargeting/Private/HardwareTargetingSettings.cpp
+++ b/Engine/Source/Editor/HardwareTargeting/Private/HardwareTargetingSettings.cpp
@@ -46,7 +46,7 @@ bool UHardwareTargetingSettings::HasPendingChanges() const
 		}
 	}
 
-	if (Settings->DefaultGraphicsPerformance != EHardwareClass::Unspecified)
+	if (Settings->DefaultGraphicsPerformance != EGraphicsPreset::Unspecified)
 	{
 		FString AppliedGraphicsEnumValue;
 		const bool bFoundKey = GConfig->GetString(TEXT("AppliedHardwareTargetingSettings"), TEXT("AppliedDefaultGraphicsPerformance"), AppliedGraphicsEnumValue, GEditorIni);
