Moved hardware targeting applied settings to DefaultEngine.ini rather than a saved Editor.ini

This ensures that multiple users will not attempt to update the project settings on start-up.

diff --git a/Engine/Source/Editor/HardwareTargeting/Private/HardwareTargetingModule.cpp b/Engine/Source/Editor/HardwareTargeting/Private/HardwareTargetingModule.cpp
--- a/Engine/Source/Editor/HardwareTargeting/Private/HardwareTargetingModule.cpp
+++ b/Engine/Source/Editor/HardwareTargeting/Private/HardwareTargetingModule.cpp
@@ -248,15 +248,9 @@ void FHardwareTargetingModule::ApplyHardwareTargetingSettings()
 		// Write out the 'did we apply' values
 		if (bSuccess)
 		{
-			UEnum* Enum = FindObject<UEnum>(ANY_PACKAGE, TEXT("EGraphicsPreset"), true);
-			check(Enum);
-			GConfig->SetString(TEXT("AppliedHardwareTargetingSettings"), TEXT("AppliedDefaultGraphicsPerformance"), *Enum->GetEnumName(Settings->DefaultGraphicsPerformance), GEditorIni);
-
-			Enum = FindObject<UEnum>(ANY_PACKAGE, TEXT("EHardwareClass"), true);
-			check(Enum);
-			GConfig->SetString(TEXT("AppliedHardwareTargetingSettings"), TEXT("AppliedTargetedHardwareClass"), *Enum->GetEnumName(Settings->TargetedHardwareClass), GEditorIni);
-
-			GConfig->Flush(false, GEditorIni);
+			Settings->AppliedTargetedHardwareClass = Settings->TargetedHardwareClass;
+			Settings->AppliedDefaultGraphicsPerformance = Settings->DefaultGraphicsPerformance;
+			Settings->UpdateDefaultConfigFile();
 		}
 	}
 }
diff --git a/Engine/Source/Editor/HardwareTargeting/Private/HardwareTargetingSettings.cpp b/Engine/Source/Editor/HardwareTargeting/Private/HardwareTargetingSettings.cpp
--- a/Engine/Source/Editor/HardwareTargeting/Private/HardwareTargetingSettings.cpp
+++ b/Engine/Source/Editor/HardwareTargeting/Private/HardwareTargetingSettings.cpp
@@ -10,71 +10,20 @@
 UHardwareTargetingSettings::UHardwareTargetingSettings(const FPostConstructInitializeProperties& PCIP)
 	: Super(PCIP)
 	, TargetedHardwareClass(EHardwareClass::Unspecified)
+	, AppliedTargetedHardwareClass(EHardwareClass::Unspecified)
 	, DefaultGraphicsPerformance(EGraphicsPreset::Unspecified)
+	, AppliedDefaultGraphicsPerformance(EGraphicsPreset::Unspecified)
 {
 }
 
 bool UHardwareTargetingSettings::HasPendingChanges() const
 {
-	UHardwareTargetingSettings* Settings = GetMutableDefault<UHardwareTargetingSettings>();
-
-	if (Settings->TargetedHardwareClass != EHardwareClass::Unspecified)
-	{
-		FString AppliedHardwareEnumValue;
-		const bool bFoundKey = GConfig->GetString(TEXT("AppliedHardwareTargetingSettings"), TEXT("AppliedTargetedHardwareClass"), AppliedHardwareEnumValue, GEditorIni);
-		if (!bFoundKey)
-		{
-			return true;
-		}
-		else if (AppliedHardwareEnumValue.IsNumeric())
-		{
-			// We used to write out an int instead of a string
-			const EHardwareClass::Type EnumValue = static_cast<EHardwareClass::Type>(FPlatformString::Atoi(*AppliedHardwareEnumValue));
-			if (EnumValue != Settings->TargetedHardwareClass)
-			{
-				return true;
-			}
-		}
-		else
-		{
-			UEnum* Enum = FindObject<UEnum>(ANY_PACKAGE, TEXT("EHardwareClass"), true);
-			check(Enum);
-			if (Enum->GetEnumName(Settings->TargetedHardwareClass) != AppliedHardwareEnumValue)
-			{
-				return true;
-			}
-		}
-	}
-
-	if (Settings->DefaultGraphicsPerformance != EGraphicsPreset::Unspecified)
+	if (TargetedHardwareClass == EHardwareClass::Unspecified || DefaultGraphicsPerformance == EGraphicsPreset::Unspecified)
 	{
-		FString AppliedGraphicsEnumValue;
-		const bool bFoundKey = GConfig->GetString(TEXT("AppliedHardwareTargetingSettings"), TEXT("AppliedDefaultGraphicsPerformance"), AppliedGraphicsEnumValue, GEditorIni);
-		if (!bFoundKey)
-		{
-			return true;
-		}
-		else if (AppliedGraphicsEnumValue.IsNumeric())
-		{
-			// We used to write out an int instead of a string
-			const EGraphicsPreset::Type EnumValue = EGraphicsPreset::Type(FPlatformString::Atoi(*AppliedGraphicsEnumValue));
-			if (EnumValue != Settings->DefaultGraphicsPerformance)
-			{
-				return true;
-			}
-		}
-		else
-		{
-			UEnum* Enum = FindObject<UEnum>(ANY_PACKAGE, TEXT("EGraphicsPreset"), true);
-			check(Enum);
-			if (Enum->GetEnumName(Settings->DefaultGraphicsPerformance) != AppliedGraphicsEnumValue)
-			{
-				return true;
-			}
-		}
+		return false;
 	}
 
-	return false;
+	return AppliedTargetedHardwareClass != TargetedHardwareClass || AppliedDefaultGraphicsPerformance != DefaultGraphicsPerformance;
 }
 
 void UHardwareTargetingSettings::PostEditChangeProperty( struct FPropertyChangedEvent& PropertyChangedEvent )
diff --git a/Engine/Source/Editor/HardwareTargeting/Public/HardwareTargetingSettings.h b/Engine/Source/Editor/HardwareTargeting/Public/HardwareTargetingSettings.h
--- a/Engine/Source/Editor/HardwareTargeting/Public/HardwareTargetingSettings.h
+++ b/Engine/Source/Editor/HardwareTargeting/Public/HardwareTargetingSettings.h
@@ -39,7 +39,7 @@ namespace EGraphicsPreset
 }
 
 /** Hardware targeting settings, stored in default config, per-project */
-UCLASS(config=Editor, defaultconfig)
+UCLASS(config=Engine, defaultconfig)
 class HARDWARETARGETING_API UHardwareTargetingSettings : public UObject
 {
 	GENERATED_UCLASS_BODY()
@@ -47,11 +47,19 @@ class HARDWARETARGETING_API UHardwareTargetingSettings : public UObject
 	/** Enum specifying which class of hardware this game is targeting */
 	UPROPERTY(config, EditAnywhere, category=None)
 	TEnumAsByte<EHardwareClass::Type> TargetedHardwareClass;
+	
+	/** Enum that is set to TargetedHardwareClass when the settings have been successfully applied */
+	UPROPERTY(config)
+	TEnumAsByte<EHardwareClass::Type> AppliedTargetedHardwareClass;
 
 	/** Enum specifying a graphics preset to use for this game */
 	UPROPERTY(config, EditAnywhere, category=None)
 	TEnumAsByte<EGraphicsPreset::Type> DefaultGraphicsPerformance;
 
+	/** Enum that is set to DefaultGraphicsPerformance when the settings have been successfully applied */
+	UPROPERTY(config)
+	TEnumAsByte<EGraphicsPreset::Type> AppliedDefaultGraphicsPerformance;
+
 	/** Check if these settings have any pending changes that require action */
 	bool HasPendingChanges() const;
 
