Editor: Fix FPropertyChangedEvent deprecation warning in HardwareTargetingModule

diff --git a/Engine/Source/Editor/HardwareTargeting/Private/HardwareTargetingModule.cpp b/Engine/Source/Editor/HardwareTargeting/Private/HardwareTargetingModule.cpp
--- a/Engine/Source/Editor/HardwareTargeting/Private/HardwareTargetingModule.cpp
+++ b/Engine/Source/Editor/HardwareTargeting/Private/HardwareTargetingModule.cpp
@@ -44,7 +44,7 @@ struct FMetaSettingGatherer
 
 			if (!bReadOnly)
 			{
-				FPropertyChangedEvent ChangeEvent(Property, false, EPropertyChangeType::ValueSet);
+				FPropertyChangedEvent ChangeEvent(Property, EPropertyChangeType::ValueSet);
 				SettingsObject->PostEditChangeProperty(ChangeEvent);
 			}
 			else
