[INTEGRATE] Change 2438570 by [REDACTED_PERSON][REDACTED_HANDLE] on 2015/02/09 14:55:13

	Editor - Improving the dropdown used in the editor for enums to have a not entirely transparent background to give it some lift from the background it's on.  Give people some visual feedback that they are buttons.

diff --git a/Engine/Source/Editor/HardwareTargeting/Public/SDecoratedEnumCombo.h b/Engine/Source/Editor/HardwareTargeting/Public/SDecoratedEnumCombo.h
--- a/Engine/Source/Editor/HardwareTargeting/Public/SDecoratedEnumCombo.h
+++ b/Engine/Source/Editor/HardwareTargeting/Public/SDecoratedEnumCombo.h
@@ -41,7 +41,7 @@ public:
 		ChildSlot
 		[
 			SNew(SComboButton)
-			.ButtonStyle(FEditorStyle::Get(), "HoverHintOnly")
+			.ButtonStyle(FEditorStyle::Get(), "FlatButton.Default")
 			.ForegroundColor(FSlateColor::UseForeground())
 			.ContentPadding(ContentPadding)
 			.OnGetMenuContent(this, &SDecoratedEnumCombo::OnGetComboContent)
