Paper2D: Enabled 2D snap mode by default when Scalable 3D/2D is picked in hardware targetting

diff --git a/Engine/Source/Editor/HardwareTargeting/Private/HardwareTargetingModule.cpp b/Engine/Source/Editor/HardwareTargeting/Private/HardwareTargetingModule.cpp
--- a/Engine/Source/Editor/HardwareTargeting/Private/HardwareTargetingModule.cpp
+++ b/Engine/Source/Editor/HardwareTargeting/Private/HardwareTargetingModule.cpp
@@ -12,7 +12,7 @@
 #include "Editor/Documentation/Public/IDocumentation.h"
 #include "GameFramework/InputSettings.h"
 #include "GameMapsSettings.h"
-
+#include "EditorProjectSettings.h"
 
 #define LOCTEXT_NAMESPACE "HardwareTargeting"
 
@@ -196,6 +196,7 @@ void FHardwareTargetingModule::GatherSettings(FMetaSettingGatherer& Builder)
 		Builder.CategoryNames.Add(GetMutableDefault<URendererSettings>(), LOCTEXT("RenderingCategoryHeader", "Engine - Rendering"));
 		Builder.CategoryNames.Add(GetMutableDefault<UInputSettings>(), LOCTEXT("InputCategoryHeader", "Engine - Input"));
 		Builder.CategoryNames.Add(GetMutableDefault<UGameMapsSettings>(), LOCTEXT("MapsAndModesCategoryHeader", "Project - Maps & Modes"));
+		Builder.CategoryNames.Add(GetMutableDefault<ULevelEditor2DSettings>(), LOCTEXT("EditorSettings2D", "Editor - 2D"));
 	}
 
 
@@ -204,6 +205,7 @@ void FHardwareTargetingModule::GatherSettings(FMetaSettingGatherer& Builder)
 	const bool bHighEndMobile = (Settings->TargetedHardwareClass == EHardwareClass::Mobile) && (Settings->DefaultGraphicsPerformance == EGraphicsPreset::Maximum);
 	const bool bAnyPC = (Settings->TargetedHardwareClass == EHardwareClass::Desktop);
 	const bool bHighEndPC = (Settings->TargetedHardwareClass == EHardwareClass::Desktop) && (Settings->DefaultGraphicsPerformance == EGraphicsPreset::Maximum);
+	const bool bAnyScalable = Settings->DefaultGraphicsPerformance == EGraphicsPreset::Scalable;
 
 	{
 		// Based roughly on https://docs.unrealengine.com/latest/INT/Platforms/Mobile/PostProcessEffects/index.html
@@ -233,6 +235,11 @@ void FHardwareTargetingModule::GatherSettings(FMetaSettingGatherer& Builder)
 		//@TODO: Use bAlwaysShowTouchInterface (sorta implied by bUseMouseForTouch)?
 	}
 
+	{
+		// Scalable automatically enables the availability of the 2D editor mode
+		UE_META_SETTING_ENTRY(Builder, ULevelEditor2DSettings, bMode2DEnabled, bAnyScalable);
+	}
+
 	{
 		// Tablets or phones are usually shared-screen multiplayer instead of split-screen
 		UE_META_SETTING_ENTRY(Builder, UGameMapsSettings, bUseSplitscreen, bAnyPC);
