Editor: Removed the snap layers/2d edit mode from being controlled by the hardware targetting setup

diff --git a/Engine/Source/Editor/HardwareTargeting/Private/HardwareTargetingModule.cpp b/Engine/Source/Editor/HardwareTargeting/Private/HardwareTargetingModule.cpp
--- a/Engine/Source/Editor/HardwareTargeting/Private/HardwareTargetingModule.cpp
+++ b/Engine/Source/Editor/HardwareTargeting/Private/HardwareTargetingModule.cpp
@@ -235,11 +235,6 @@ void FHardwareTargetingModule::GatherSettings(FMetaSettingGatherer& Builder)
 		//@TODO: Use bAlwaysShowTouchInterface (sorta implied by bUseMouseForTouch)?
 	}
 
-	{
-		// Scalable automatically enables the availability of the 2D editor mode
-		UE_META_SETTING_ENTRY(Builder, ULevelEditor2DSettings, bMode2DEnabled, bAnyScalable);
-	}
-
 	{
 		// Tablets or phones are usually shared-screen multiplayer instead of split-screen
 		UE_META_SETTING_ENTRY(Builder, UGameMapsSettings, bUseSplitscreen, bAnyPC);
