fixed LensFlares on by default (project setting), only affects newly created projects

diff --git a/Engine/Source/Editor/HardwareTargeting/Private/HardwareTargetingModule.cpp b/Engine/Source/Editor/HardwareTargeting/Private/HardwareTargetingModule.cpp
--- a/Engine/Source/Editor/HardwareTargeting/Private/HardwareTargetingModule.cpp
+++ b/Engine/Source/Editor/HardwareTargeting/Private/HardwareTargetingModule.cpp
@@ -217,12 +217,14 @@ void FHardwareTargetingModule::GatherSettings(FMetaSettingGatherer& Builder)
 		// Separate translucency does nothing in the ES2 renderer
 		UE_META_SETTING_ENTRY(Builder, URendererSettings, bSeparateTranslucency, !bAnyMobile);
 
-		// Motion blur, lens flare, auto-exposure, and ambient occlusion don't work in the ES2 renderer
+		// Motion blur, auto-exposure, and ambient occlusion don't work in the ES2 renderer
 		UE_META_SETTING_ENTRY(Builder, URendererSettings, bDefaultFeatureMotionBlur, bHighEndPC);
-		UE_META_SETTING_ENTRY(Builder, URendererSettings, bDefaultFeatureLensFlare, bAnyPC);
 		UE_META_SETTING_ENTRY(Builder, URendererSettings, bDefaultFeatureAutoExposure, bHighEndPC);
 		UE_META_SETTING_ENTRY(Builder, URendererSettings, bDefaultFeatureAmbientOcclusion, bAnyPC);
 
+		// lens flare doesn't work in the ES2 renderer, the quality is low and the feature is controversial
+		UE_META_SETTING_ENTRY(Builder, URendererSettings, bDefaultFeatureLensFlare, false);
+
 		// DOF and AA work on mobile but are expensive, keeping them off by default
 		//@TODO: DOF setting doesn't exist yet
 		// UE_META_SETTING_ENTRY(Builder, URendererSettings, bDefaultFeatureDepthOfField, bHighEndPC);
