It's now possible to flip SDecoratedEnumCombo to vertical, which was necessary for the New Project dialog rewrite.

diff --git a/Engine/Source/Editor/DetailCustomizations/Private/HardwareTargetingSettingsDetails.cpp b/Engine/Source/Editor/DetailCustomizations/Private/HardwareTargetingSettingsDetails.cpp
--- a/Engine/Source/Editor/DetailCustomizations/Private/HardwareTargetingSettingsDetails.cpp
+++ b/Engine/Source/Editor/DetailCustomizations/Private/HardwareTargetingSettingsDetails.cpp
@@ -286,7 +286,8 @@ void FHardwareTargetingSettingsDetails::CustomizeDetails(IDetailLayoutBuilder& D
 		
 		HardwareClassCombo = HardwareTargeting.MakeHardwareClassTargetCombo(
 			FOnHardwareClassChanged::CreateStatic(SetPropertyValue, Property),
-			TAttribute<EHardwareClass::Type>::Create(TAttribute<EHardwareClass::Type>::FGetter::CreateStatic(GetPropertyValue, Property))
+			TAttribute<EHardwareClass::Type>::Create(TAttribute<EHardwareClass::Type>::FGetter::CreateStatic(GetPropertyValue, Property)),
+			Orient_Horizontal
 		);
 	}
 
@@ -306,7 +307,8 @@ void FHardwareTargetingSettingsDetails::CustomizeDetails(IDetailLayoutBuilder& D
 		};
 		GraphicsPresetCombo = HardwareTargeting.MakeGraphicsPresetTargetCombo(
 			FOnGraphicsPresetChanged::CreateStatic(SetPropertyValue, Property),
-			TAttribute<EGraphicsPreset::Type>::Create(TAttribute<EGraphicsPreset::Type>::FGetter::CreateStatic(GetPropertyValue, Property))
+			TAttribute<EGraphicsPreset::Type>::Create(TAttribute<EGraphicsPreset::Type>::FGetter::CreateStatic(GetPropertyValue, Property)),
+			Orient_Horizontal
 		);
 	}
 
diff --git a/Engine/Source/Editor/HardwareTargeting/Private/HardwareTargetingModule.cpp b/Engine/Source/Editor/HardwareTargeting/Private/HardwareTargetingModule.cpp
--- a/Engine/Source/Editor/HardwareTargeting/Private/HardwareTargetingModule.cpp
+++ b/Engine/Source/Editor/HardwareTargeting/Private/HardwareTargetingModule.cpp
@@ -148,8 +148,8 @@ public:
 	// IHardwareTargetingModule interface
 	virtual void ApplyHardwareTargetingSettings() override;
 	virtual TArray<FModifiedDefaultConfig> GetPendingSettingsChanges() override;
-	virtual TSharedRef<SWidget> MakeHardwareClassTargetCombo(FOnHardwareClassChanged OnChanged, TAttribute<EHardwareClass::Type> SelectedEnum) override;
-	virtual TSharedRef<SWidget> MakeGraphicsPresetTargetCombo(FOnGraphicsPresetChanged OnChanged, TAttribute<EGraphicsPreset::Type> SelectedEnum) override;
+	virtual TSharedRef<SWidget> MakeHardwareClassTargetCombo(FOnHardwareClassChanged OnChanged, TAttribute<EHardwareClass::Type> SelectedEnum, EOrientation Orientation) override;
+	virtual TSharedRef<SWidget> MakeGraphicsPresetTargetCombo(FOnGraphicsPresetChanged OnChanged, TAttribute<EGraphicsPreset::Type> SelectedEnum, EOrientation Orientation) override;
 	// End of IHardwareTargetingModule interface
 
 private:
@@ -280,7 +280,7 @@ void FHardwareTargetingModule::ApplyHardwareTargetingSettings()
 	}
 }
 
-TSharedRef<SWidget> FHardwareTargetingModule::MakeHardwareClassTargetCombo(FOnHardwareClassChanged OnChanged, TAttribute<EHardwareClass::Type> SelectedEnum)
+TSharedRef<SWidget> FHardwareTargetingModule::MakeHardwareClassTargetCombo(FOnHardwareClassChanged OnChanged, TAttribute<EHardwareClass::Type> SelectedEnum, EOrientation Orientation)
 {
 	TArray<SDecoratedEnumCombo<EHardwareClass::Type>::FComboOption> HardwareClassInfo;
 	HardwareClassInfo.Add(SDecoratedEnumCombo<EHardwareClass::Type>::FComboOption(
@@ -293,10 +293,11 @@ TSharedRef<SWidget> FHardwareTargetingModule::MakeHardwareClassTargetCombo(FOnHa
 	return SNew(SDecoratedEnumCombo<EHardwareClass::Type>, MoveTemp(HardwareClassInfo))
 		.SelectedEnum(SelectedEnum)
 		.OnEnumChanged(OnChanged)
+		.Orientation(Orientation)
 		.ToolTip(IDocumentation::Get()->CreateToolTip(LOCTEXT("HardwareClassTooltip", "Choose the overall class of hardware to target (desktop/console or mobile/tablet)."), NULL, TEXT("Shared/Editor/Settings/TargetHardware"), TEXT("HardwareClass")));
 }
 
-TSharedRef<SWidget> FHardwareTargetingModule::MakeGraphicsPresetTargetCombo(FOnGraphicsPresetChanged OnChanged, TAttribute<EGraphicsPreset::Type> SelectedEnum)
+TSharedRef<SWidget> FHardwareTargetingModule::MakeGraphicsPresetTargetCombo(FOnGraphicsPresetChanged OnChanged, TAttribute<EGraphicsPreset::Type> SelectedEnum, EOrientation Orientation)
 {
 	TArray<SDecoratedEnumCombo<EGraphicsPreset::Type>::FComboOption> GraphicsPresetInfo;
 	GraphicsPresetInfo.Add(SDecoratedEnumCombo<EGraphicsPreset::Type>::FComboOption(
@@ -309,6 +310,7 @@ TSharedRef<SWidget> FHardwareTargetingModule::MakeGraphicsPresetTargetCombo(FOnG
 	return SNew(SDecoratedEnumCombo<EGraphicsPreset::Type>, MoveTemp(GraphicsPresetInfo))
 		.SelectedEnum(SelectedEnum)
 		.OnEnumChanged(OnChanged)
+		.Orientation(Orientation)
 		.ToolTip(IDocumentation::Get()->CreateToolTip(LOCTEXT("GraphicsPresetTooltip", "Choose the graphical level to target (high-end only or scalable from low-end on up)."), NULL, TEXT("Shared/Editor/Settings/TargetHardware"), TEXT("GraphicalLevel")));
 }
 
diff --git a/Engine/Source/Editor/HardwareTargeting/Public/HardwareTargetingModule.h b/Engine/Source/Editor/HardwareTargeting/Public/HardwareTargetingModule.h
--- a/Engine/Source/Editor/HardwareTargeting/Public/HardwareTargetingModule.h
+++ b/Engine/Source/Editor/HardwareTargeting/Public/HardwareTargetingModule.h
@@ -7,6 +7,7 @@
 #include "UObject/WeakObjectPtr.h"
 #include "Modules/ModuleInterface.h"
 #include "HardwareTargetingSettings.h"
+#include "Types/SlateEnums.h"
 
 class SWidget;
 
@@ -40,9 +41,9 @@ public:
 	virtual TArray<FModifiedDefaultConfig> GetPendingSettingsChanges() = 0;
 
 	/** Make a new combo box for choosing a hardware class target */
-	virtual TSharedRef<SWidget> MakeHardwareClassTargetCombo(FOnHardwareClassChanged OnChanged, TAttribute<EHardwareClass::Type> SelectedEnum) = 0;
+	virtual TSharedRef<SWidget> MakeHardwareClassTargetCombo(FOnHardwareClassChanged OnChanged, TAttribute<EHardwareClass::Type> SelectedEnum, EOrientation Orientation) = 0;
 
 	/** Make a new combo box for choosing a graphics preference */
-	virtual TSharedRef<SWidget> MakeGraphicsPresetTargetCombo(FOnGraphicsPresetChanged OnChanged, TAttribute<EGraphicsPreset::Type> SelectedEnum) = 0;
+	virtual TSharedRef<SWidget> MakeGraphicsPresetTargetCombo(FOnGraphicsPresetChanged OnChanged, TAttribute<EGraphicsPreset::Type> SelectedEnum, EOrientation Orientation) = 0;
 };
 
diff --git a/Engine/Source/Editor/HardwareTargeting/Public/SDecoratedEnumCombo.h b/Engine/Source/Editor/HardwareTargeting/Public/SDecoratedEnumCombo.h
--- a/Engine/Source/Editor/HardwareTargeting/Public/SDecoratedEnumCombo.h
+++ b/Engine/Source/Editor/HardwareTargeting/Public/SDecoratedEnumCombo.h
@@ -19,6 +19,7 @@
 #include "Widgets/Text/STextBlock.h"
 #include "Widgets/Input/SButton.h"
 #include "Widgets/Input/SComboButton.h"
+#include "Types/SlateEnums.h"
 
 template<typename TEnumType>
 class SDecoratedEnumCombo : public SCompoundWidget
@@ -40,10 +41,12 @@ public:
 
 	SLATE_BEGIN_ARGS( SDecoratedEnumCombo )
 		: _ContentPadding(6.f)
+		, _Orientation(Orient_Horizontal)
 	{}
 
 		SLATE_EVENT(FOnEnumChanged, OnEnumChanged)
 		SLATE_ARGUMENT(FMargin, ContentPadding)
+		SLATE_ARGUMENT(EOrientation, Orientation)
 		SLATE_ATTRIBUTE(TEnumType, SelectedEnum)
 
 	SLATE_END_ARGS()
@@ -55,6 +58,7 @@ public:
 		OnEnumChanged = InArgs._OnEnumChanged;
 		Options = MoveTemp(InOptions);
 		SelectedEnum = InArgs._SelectedEnum;
+		Orientation = InArgs._Orientation;
 
 		ChildSlot
 		[
@@ -94,6 +98,7 @@ private:
 	FMargin ContentPadding;
 	FOnEnumChanged OnEnumChanged;
 	TWeakPtr<SWidget> MenuContent;
+	EOrientation Orientation;
 
 	FText GetCurrentText() const
 	{
@@ -133,7 +138,16 @@ private:
 
 	TSharedRef<SWidget> OnGetComboContent()
 	{
-		TSharedRef<SHorizontalBox> HorizontalBox = SNew(SHorizontalBox);
+		TSharedPtr<SVerticalBox> VerticalBox;
+		TSharedPtr<SHorizontalBox> HorizontalBox;
+		if (Orientation == Orient_Vertical)
+		{
+			VerticalBox = SNew(SVerticalBox);
+		}
+		else
+		{
+			HorizontalBox = SNew(SHorizontalBox);
+		}
 
 		for (int32 Index = 0; Index < Options.Num(); ++Index)
 		{
@@ -143,8 +157,17 @@ private:
 				continue;
 			}
 
-			HorizontalBox->AddSlot()
-			.AutoWidth()
+			SBoxPanel::FSlot* Slot;
+			if (Orientation == Orient_Vertical)
+			{
+				Slot = &VerticalBox->AddSlot().AutoHeight();
+			}
+			else
+			{
+				Slot = &HorizontalBox->AddSlot().AutoWidth();
+			}
+			
+			(*Slot)
 			[
 				SNew(SButton)
 				.ButtonStyle(FEditorStyle::Get(), "HoverHintOnly")
@@ -172,8 +195,8 @@ private:
 			];
 		}
 
-		MenuContent = HorizontalBox;
+		MenuContent = Orientation == Orient_Vertical ? StaticCastSharedPtr<SWidget>(VerticalBox) : StaticCastSharedPtr<SWidget>(HorizontalBox);
 
-		return HorizontalBox;
+		return MenuContent.Pin().ToSharedRef();
 	}
 };
